Fix crash using the facial anim importer

diff --git a/Engine/Plugins/Editor/FacialAnimation/Source/FacialAnimationEditor/Private/FacialAnimationImportItem.cpp b/Engine/Plugins/Editor/FacialAnimation/Source/FacialAnimationEditor/Private/FacialAnimationImportItem.cpp
--- a/Engine/Plugins/Editor/FacialAnimation/Source/FacialAnimationEditor/Private/FacialAnimationImportItem.cpp
+++ b/Engine/Plugins/Editor/FacialAnimation/Source/FacialAnimationEditor/Private/FacialAnimationImportItem.cpp
@@ -89,7 +89,7 @@ bool FFacialAnimationImportItem::ImportCurvesEmbeddedInSoundWave()
 		if (FbxAnimUtils::ImportCurveTableFromNode(FbxFile, GetDefault<UFacialAnimationBulkImporterSettings>()->CurveNodeName, NewCurveTable, PreRollTime))
 		{
 			// we will need to add a curve to tell us the time we want to start playing audio
-			FSimpleCurve& AudioCurve = NewCurveTable->AddSimpleCurve(TEXT("Audio"));
+			FRichCurve& AudioCurve = NewCurveTable->AddRichCurve(TEXT("Audio"));
 			AudioCurve.AddKey(PreRollTime, 1.0f);
 
 			return true;
