Trivial: Fix mis-matched optional parameter const correctness causing game target build warnings

diff --git a/Engine/Plugins/Editor/FacialAnimation/Source/FacialAnimation/Private/AudioCurveSourceComponent.cpp b/Engine/Plugins/Editor/FacialAnimation/Source/FacialAnimation/Private/AudioCurveSourceComponent.cpp
--- a/Engine/Plugins/Editor/FacialAnimation/Source/FacialAnimation/Private/AudioCurveSourceComponent.cpp
+++ b/Engine/Plugins/Editor/FacialAnimation/Source/FacialAnimation/Private/AudioCurveSourceComponent.cpp
@@ -55,7 +55,7 @@ void UAudioCurveSourceComponent::CacheCurveData()
 	}
 }
 
-void UAudioCurveSourceComponent::FadeIn(float FadeInDuration, float FadeVolumeLevel, float StartTime, EAudioFaderCurve FadeType)
+void UAudioCurveSourceComponent::FadeIn(float FadeInDuration, float FadeVolumeLevel, float StartTime, const EAudioFaderCurve FadeType)
 {
 	CacheCurveData();
 
@@ -73,7 +73,7 @@ void UAudioCurveSourceComponent::FadeIn(float FadeInDuration, float FadeVolumeLe
 	}
 }
 
-void UAudioCurveSourceComponent::FadeOut(float FadeOutDuration, float FadeVolumeLevel, EAudioFaderCurve FadeType)
+void UAudioCurveSourceComponent::FadeOut(float FadeOutDuration, float FadeVolumeLevel, const EAudioFaderCurve FadeType)
 {
 	if (Delay < CachedSyncPreRoll)
 	{
diff --git a/Engine/Plugins/Editor/FacialAnimation/Source/FacialAnimation/Public/AudioCurveSourceComponent.h b/Engine/Plugins/Editor/FacialAnimation/Source/FacialAnimation/Public/AudioCurveSourceComponent.h
--- a/Engine/Plugins/Editor/FacialAnimation/Source/FacialAnimation/Public/AudioCurveSourceComponent.h
+++ b/Engine/Plugins/Editor/FacialAnimation/Source/FacialAnimation/Public/AudioCurveSourceComponent.h
@@ -37,8 +37,8 @@ public:
 	virtual void TickComponent(float DeltaTime, enum ELevelTick TickType, FActorComponentTickFunction *ThisTickFunction) override;
 
 	/** UAudioComponent interface */
-	virtual void FadeIn(float FadeInDuration, float FadeVolumeLevel = 1.0f, float StartTime = 0.0f, EAudioFaderCurve FadeType = EAudioFaderCurve::Linear) override;
-	virtual	void FadeOut(float FadeOutDuration, float FadeVolumeLevel, EAudioFaderCurve FadeType = EAudioFaderCurve::Linear) override;
+	virtual void FadeIn(float FadeInDuration, float FadeVolumeLevel = 1.0f, float StartTime = 0.0f, const EAudioFaderCurve FadeType = EAudioFaderCurve::Linear) override;
+	virtual	void FadeOut(float FadeOutDuration, float FadeVolumeLevel, const EAudioFaderCurve FadeType = EAudioFaderCurve::Linear) override;
 	virtual void Play(float StartTime = 0.0f) override;
 	virtual void Stop() override;
 	virtual bool IsPlaying() const override;
